OSSSteam: Turn off the accidental setting of steam protocol as the default socketsubsystem.
#rnx

diff --git a/Engine/Plugins/Online/SocketSubsystemSteamIP/Source/SocketSubsystemSteamIP/Private/SocketSubsystemSteamIPModule.cpp b/Engine/Plugins/Online/SocketSubsystemSteamIP/Source/SocketSubsystemSteamIP/Private/SocketSubsystemSteamIPModule.cpp
--- a/Engine/Plugins/Online/SocketSubsystemSteamIP/Source/SocketSubsystemSteamIP/Private/SocketSubsystemSteamIPModule.cpp
+++ b/Engine/Plugins/Online/SocketSubsystemSteamIP/Source/SocketSubsystemSteamIP/Private/SocketSubsystemSteamIPModule.cpp
@@ -18,7 +18,7 @@ void FSocketSubsystemSteamIPModule::StartupModule()
 	if (SharedModule.AreSteamDllsLoaded() && bIsNotEditor)
 	{
 		// Settings flags
-		bool bOverrideSocketSubsystem = true;
+		bool bOverrideSocketSubsystem = false;
 
 		// Use this flag from the SteamNetworking configuration so that the IPNetDrivers can be used as needed.
 		if (GConfig)
